[AVCodecs] Fix Infinite KeyFrameLength with AMF causeing error messages.

NVEnc and AMF use different value to indicate infinite KeyFrameLength. AMF only accepts a value between (inclusive) 0 and 1000, with 0 indicating infinite KeyFrameLength.

This fix converts from NVEnc 0xffffffff to 0 inside AMF.

diff --git a/Engine/Plugins/Experimental/AVCodecs/AMFCodecs/Source/AMFCodecs/Private/Video/Encoders/Configs/VideoEncoderConfigAMF.cpp b/Engine/Plugins/Experimental/AVCodecs/AMFCodecs/Source/AMFCodecs/Private/Video/Encoders/Configs/VideoEncoderConfigAMF.cpp
--- a/Engine/Plugins/Experimental/AVCodecs/AMFCodecs/Source/AMFCodecs/Private/Video/Encoders/Configs/VideoEncoderConfigAMF.cpp
+++ b/Engine/Plugins/Experimental/AVCodecs/AMFCodecs/Source/AMFCodecs/Private/Video/Encoders/Configs/VideoEncoderConfigAMF.cpp
@@ -87,7 +87,10 @@ DLLEXPORT FAVResult FAVExtension::TransformConfig(FVideoEncoderConfigAMF& OutCon
 
 	OutConfig.SetProperty(AMF_VIDEO_ENCODER_QUERY_TIMEOUT, 16);
 
-	OutConfig.SetProperty(AMF_VIDEO_ENCODER_IDR_PERIOD, InConfig.KeyframeInterval);
+	// Pixel streaming uses Nvidia convention of using 0xffffffff for infinite p-frames.
+	// AMD uses 0. See AMF_VIDEO_ENCODER_IDR_PERIOD in [REDACTED_PROFILE_URL]/AMF/blob/master/amf/doc/AMF_Video_Encode_API.md
+	uint32 KeyframeInterval = InConfig.KeyframeInterval == 0xffffffff ? 0 : InConfig.KeyframeInterval;
+	OutConfig.SetProperty(AMF_VIDEO_ENCODER_IDR_PERIOD, KeyframeInterval);
 
 	return EAVResult::Success;
 }
@@ -136,7 +139,10 @@ DLLEXPORT FAVResult FAVExtension::TransformConfig(FVideoEncoderConfig& OutConfig
 	InConfig.GetProperty(AMF_VIDEO_ENCODER_MAX_QP, &OutMaxQpValue);
 	OutConfig.MinQuality = (1.0f - (OutMaxQpValue / 51.0f)) * 100.0f;
 
+	// Pixel streaming uses Nvidia convention of using 0xffffffff for infinite p-frames.
+	// AMD uses 0. See AMF_VIDEO_ENCODER_IDR_PERIOD in [REDACTED_PROFILE_URL]/AMF/blob/master/amf/doc/AMF_Video_Encode_API.md
 	InConfig.GetProperty(AMF_VIDEO_ENCODER_IDR_PERIOD, &OutConfig.KeyframeInterval);
+	OutConfig.KeyframeInterval = OutConfig.KeyframeInterval > 0 ? OutConfig.KeyframeInterval : 0xffffffff;
 
 	return EAVResult::Success;
 }
