[AVCodecs] Attempt to fix build issue that appears in horde by removing h265 encode for AMF.

Horde seems to timeout, and the only real clue is FAVExtension::TransformConfig<FVideoEncoderConfigAMF,FVideoEncoderConfigH265>.

Since h265 is no longer used, this CL removes H265 and TransformConfig for h265 for AMF.

diff --git a/Engine/Plugins/Experimental/AVCodecs/AMFCodecs/Source/AMFCodecs/Private/AMFCodecsModule.cpp b/Engine/Plugins/Experimental/AVCodecs/AMFCodecs/Source/AMFCodecs/Private/AMFCodecsModule.cpp
--- a/Engine/Plugins/Experimental/AVCodecs/AMFCodecs/Source/AMFCodecs/Private/AMFCodecsModule.cpp
+++ b/Engine/Plugins/Experimental/AVCodecs/AMFCodecs/Source/AMFCodecs/Private/AMFCodecsModule.cpp
@@ -22,16 +22,17 @@ public:
 	virtual void StartupModule() override
 	{
 #if PLATFORM_WINDOWS
+		// TODO (Eden.Harris) H265/HEVC disabled because of intermittent build failure. See UE-239424
 		FVideoEncoder::RegisterPermutationsOf<TVideoEncoderAMF<FVideoResourceD3D11>>
 			::With<FVideoResourceD3D11>
-			::And<FVideoEncoderConfigAMF, FVideoEncoderConfigH264, FVideoEncoderConfigH265>([](TSharedRef<FAVDevice> const& NewDevice, TSharedRef<FAVInstance> const& NewInstance)
+			::And<FVideoEncoderConfigAMF, FVideoEncoderConfigH264>([](TSharedRef<FAVDevice> const& NewDevice, TSharedRef<FAVInstance> const& NewInstance)
 				{
 					return FAPI::Get<FAMF>().IsValid();
 				});
 
 		FVideoEncoder::RegisterPermutationsOf<TVideoEncoderAMF<FVideoResourceD3D12>>
 			::With<FVideoResourceD3D12>
-			::And<FVideoEncoderConfigAMF, FVideoEncoderConfigH264, FVideoEncoderConfigH265>([](TSharedRef<FAVDevice> const& NewDevice, TSharedRef<FAVInstance> const& NewInstance)
+			::And<FVideoEncoderConfigAMF, FVideoEncoderConfigH264>([](TSharedRef<FAVDevice> const& NewDevice, TSharedRef<FAVInstance> const& NewInstance)
 				{
 					return FAPI::Get<FAMF>().IsValid();
 				});
diff --git a/Engine/Plugins/Experimental/AVCodecs/AMFCodecs/Source/AMFCodecs/Private/Video/Encoders/Configs/VideoEncoderConfigAMF.cpp b/Engine/Plugins/Experimental/AVCodecs/AMFCodecs/Source/AMFCodecs/Private/Video/Encoders/Configs/VideoEncoderConfigAMF.cpp
--- a/Engine/Plugins/Experimental/AVCodecs/AMFCodecs/Source/AMFCodecs/Private/Video/Encoders/Configs/VideoEncoderConfigAMF.cpp
+++ b/Engine/Plugins/Experimental/AVCodecs/AMFCodecs/Source/AMFCodecs/Private/Video/Encoders/Configs/VideoEncoderConfigAMF.cpp
@@ -190,33 +190,34 @@ DLLEXPORT FAVResult FAVExtension::TransformConfig(FVideoEncoderConfigAMF& OutCon
 	return FAVExtension::TransformConfig<FVideoEncoderConfigAMF, FVideoEncoderConfig>(OutConfig, InConfig);
 }
 
-template <>
-DLLEXPORT FAVResult FAVExtension::TransformConfig(FVideoEncoderConfigAMF& OutConfig, FVideoEncoderConfigH265 const& InConfig)
-{
-	OutConfig.CodecType = FVideoEncoderConfigAMF::CodecTypeH265;
-
-	static auto const ConvertProfile = [](EH265Profile Profile) -> TAVResult<AMF_VIDEO_ENCODER_PROFILE_ENUM>
-		{
-			switch (Profile)
-			{
-			case EH265Profile::Main:
-				return AMF_VIDEO_ENCODER_PROFILE_ENUM::AMF_VIDEO_ENCODER_PROFILE_MAIN;
-			default:
-				return FAVResult(EAVResult::ErrorUnsupported, FString::Printf(TEXT("H265 profile %d is not supported"), Profile), TEXT("AMF"));
-			}
-		};
-
-	OutConfig.CodecType = FVideoEncoderConfigAMF::CodecTypeH265;
-	OutConfig.RepeatSPSPPS = InConfig.RepeatSPSPPS;
-
-	TAVResult<AMF_VIDEO_ENCODER_PROFILE_ENUM> const ConvertedProfile = ConvertProfile(InConfig.Profile);
-	if (ConvertedProfile.IsNotSuccess())
-	{
-		return ConvertedProfile;
-	}
-
-	OutConfig.SetProperty(AMF_VIDEO_ENCODER_PROFILE, ConvertedProfile.ReturnValue);
-	OutConfig.SetProperty(AMF_VIDEO_ENCODER_PROFILE_LEVEL, 51);
-
-	return FAVExtension::TransformConfig<FVideoEncoderConfigAMF, FVideoEncoderConfig>(OutConfig, InConfig);
-}
+// TODO (Eden.Harris) H265/HEVC disabled because of intermittent build failure. See UE-239424
+// template <>
+// DLLEXPORT FAVResult FAVExtension::TransformConfig(FVideoEncoderConfigAMF& OutConfig, FVideoEncoderConfigH265 const& InConfig)
+// {
+// 	OutConfig.CodecType = FVideoEncoderConfigAMF::CodecTypeH265;
+
+// 	static auto const ConvertProfile = [](EH265Profile Profile) -> TAVResult<AMF_VIDEO_ENCODER_PROFILE_ENUM>
+// 		{
+// 			switch (Profile)
+// 			{
+// 			case EH265Profile::Main:
+// 				return AMF_VIDEO_ENCODER_PROFILE_ENUM::AMF_VIDEO_ENCODER_PROFILE_MAIN;
+// 			default:
+// 				return FAVResult(EAVResult::ErrorUnsupported, FString::Printf(TEXT("H265 profile %d is not supported"), Profile), TEXT("AMF"));
+// 			}
+// 		};
+
+// 	OutConfig.CodecType = FVideoEncoderConfigAMF::CodecTypeH265;
+// 	OutConfig.RepeatSPSPPS = InConfig.RepeatSPSPPS;
+
+// 	TAVResult<AMF_VIDEO_ENCODER_PROFILE_ENUM> const ConvertedProfile = ConvertProfile(InConfig.Profile);
+// 	if (ConvertedProfile.IsNotSuccess())
+// 	{
+// 		return ConvertedProfile;
+// 	}
+
+// 	OutConfig.SetProperty(AMF_VIDEO_ENCODER_PROFILE, ConvertedProfile.ReturnValue);
+// 	OutConfig.SetProperty(AMF_VIDEO_ENCODER_PROFILE_LEVEL, 51);
+
+// 	return FAVExtension::TransformConfig<FVideoEncoderConfigAMF, FVideoEncoderConfig>(OutConfig, InConfig);
+// }
diff --git a/Engine/Plugins/Experimental/AVCodecs/AMFCodecs/Source/AMFCodecs/Public/Video/Encoders/Configs/VideoEncoderConfigAMF.h b/Engine/Plugins/Experimental/AVCodecs/AMFCodecs/Source/AMFCodecs/Public/Video/Encoders/Configs/VideoEncoderConfigAMF.h
--- a/Engine/Plugins/Experimental/AVCodecs/AMFCodecs/Source/AMFCodecs/Public/Video/Encoders/Configs/VideoEncoderConfigAMF.h
+++ b/Engine/Plugins/Experimental/AVCodecs/AMFCodecs/Source/AMFCodecs/Public/Video/Encoders/Configs/VideoEncoderConfigAMF.h
@@ -164,7 +164,8 @@ FAVResult FAVExtension::TransformConfig(struct FVideoEncoderConfig& OutConfig, F
 template <>
 FAVResult FAVExtension::TransformConfig(FVideoEncoderConfigAMF& OutConfig, struct FVideoEncoderConfigH264 const& InConfig);
 
-template <>
-FAVResult FAVExtension::TransformConfig(FVideoEncoderConfigAMF& OutConfig, struct FVideoEncoderConfigH265 const& InConfig);
+// TODO (Eden.Harris) H265/HEVC disabled because of intermittent build failure. See UE-239424
+// template <>
+// FAVResult FAVExtension::TransformConfig(FVideoEncoderConfigAMF& OutConfig, struct FVideoEncoderConfigH265 const& InConfig);
 
 DECLARE_TYPEID(FVideoEncoderConfigAMF, AMFCODECS_API);
